[FlutterCommon] [LiveLinkVCAM] Split Tentacle handling into its own package

diff --git a/Engine/Extras/FlutterTentacle/.gitignore b/Engine/Extras/FlutterTentacle/.gitignore
new file mode 100644
--- /dev/null
+++ b/Engine/Extras/FlutterTentacle/.gitignore
@@ -0,0 +1,29 @@
+# Miscellaneous
+*.class
+*.log
+*.pyc
+*.swp
+.DS_Store
+.atom/
+.buildlog/
+.history
+.svn/
+migrate_working_dir/
+
+# IntelliJ related
+*.iml
+*.ipr
+*.iws
+.idea/
+
+# The .vscode folder contains launch configuration and tasks you configure in
+# VS Code which you may wish to be included in version control, so this line
+# is commented out by default.
+#.vscode/
+
+# Flutter/Dart/Pub related
+# Libraries should not include pubspec.lock, per https://dart.dev/guides/libraries/private-files#pubspeclock.
+/pubspec.lock
+**/doc/api/
+.dart_tool/
+build/
diff --git a/Engine/Extras/FlutterTentacle/IOS/.gitignore b/Engine/Extras/FlutterTentacle/IOS/.gitignore
new file mode 100644
--- /dev/null
+++ b/Engine/Extras/FlutterTentacle/IOS/.gitignore
@@ -0,0 +1,38 @@
+.idea/
+.vagrant/
+.sconsign.dblite
+.svn/
+
+.DS_Store
+*.swp
+profile
+
+DerivedData/
+build/
+GeneratedPluginRegistrant.h
+GeneratedPluginRegistrant.m
+
+.generated/
+
+*.pbxuser
+*.mode1v3
+*.mode2v3
+*.perspectivev3
+
+!default.pbxuser
+!default.mode1v3
+!default.mode2v3
+!default.perspectivev3
+
+xcuserdata
+
+*.moved-aside
+
+*.pyc
+*sync/
+Icon?
+.tags*
+
+/Flutter/Generated.xcconfig
+/Flutter/ephemeral/
+/Flutter/flutter_export_environment.sh
diff --git a/Engine/Extras/FlutterTentacle/IOS/Classes/API/FlutterTentacleApi.swift b/Engine/Extras/FlutterTentacle/IOS/Classes/API/FlutterTentacleApi.swift
new file mode 100644
--- /dev/null
+++ b/Engine/Extras/FlutterTentacle/IOS/Classes/API/FlutterTentacleApi.swift
@@ -0,0 +1,90 @@
+// Copyright Epic Games, Inc. All Rights Reserved.
+
+import Flutter
+import Foundation
+import Tentacle
+
+/// Handles messages about Tentacle to and from the Flutter API.
+class FlutterTentacleApi: FlutterPluginApi<TentacleFlutterApi> {
+  /// The Bluetooth controller that scans for Tentacle devices.
+  private var bluetooth: TentacleBluetoothController!
+  
+  /// Whether this should attempt to scan for Tentacle devices.
+  private var bIsScanning: Bool = false
+  
+  init(binaryMessenger: FlutterBinaryMessenger) {
+    super.init(flutter:TentacleFlutterApi(binaryMessenger: binaryMessenger))
+    
+    TentacleHostApiSetup.setUp(binaryMessenger: binaryMessenger, api: self)
+    FlutterTentacleApi.instance = self
+    
+    bluetooth = TentacleBluetoothController(DispatchQueue.main)
+    bluetooth.delegate = self
+  }
+  
+  deinit {
+    if (FlutterTentacleApi.instance === self) {
+      FlutterTentacleApi.instance = nil
+    }
+  }
+  
+  /// The active singleton instance of this class.
+  private(set) static var instance: FlutterTentacleApi?
+  
+  /// Convert a C string with the given length to a Swift string.
+  private func stringFromCString<T>(cString: inout T, length: Int) -> String {
+    return withUnsafePointer(to: cString) { startPointer in
+      return startPointer.withMemoryRebound(to: UInt8.self, capacity: MemoryLayout.size(ofValue: length)) { cStringPointer in
+        return String(cString: cStringPointer)
+      }
+    }
+  }
+}
+
+// MARK: Host API
+extension FlutterTentacleApi: TentacleHostApi {
+  func startScanning() throws {
+    bIsScanning = true
+    bluetooth.startScanning()
+  }
+  
+  func stopScanning() throws {
+    bIsScanning = false
+    bluetooth.stopScanning()
+  }
+  
+  func getCachedDeviceInfo() throws -> [TentacleDeviceInfo] {
+    var devices: [TentacleDeviceInfo] = []
+    let deviceCount: Int32 = Int32(Tentacle.TentacleDeviceCacheGetSize())
+  
+    for deviceIndex in stride(from: 0, to: deviceCount, by: 1) {
+      let device: TentacleDevice = Tentacle.TentacleDeviceCacheGetDevice(deviceIndex)
+      var advert: TentacleAdvertisement = device.advertisement
+      
+      devices.append(TentacleDeviceInfo(
+        bIsInGreenMode: advert.greenMode,
+        bIsCharging: advert.charging,
+        bIsDropFrame: advert.dropFrame,
+        batteryLevel: Int64(advert.battery),
+        frameRate: Tentacle.TentacleAdvertisementGetFrameRate(&advert),
+        lastSeenTimestamp: CACurrentMediaTime(),
+        iconIndex: Int64(advert.icon),
+        signalStrength: Int64(advert.rssi),
+        productId: advert.productId.toFlutter(),
+        identifier: stringFromCString(cString: &advert.identifier, length: advert.identifierLength),
+        name: stringFromCString(cString: &advert.name, length: advert.nameLength)
+      ))
+    }
+    
+    return devices
+  }
+}
+
+// MARK: Bluetooth Controller Delegate
+extension FlutterTentacleApi : TentacleBluetoothControllerDelegate {
+  func didUpdate(to state: TentacleBluetoothState) {
+    if state == .poweredOn && bIsScanning {
+      bluetooth.startScanning()
+    }
+  }
+}
diff --git a/Engine/Extras/FlutterTentacle/IOS/Classes/FlutterTentaclePlugin.swift b/Engine/Extras/FlutterTentacle/IOS/Classes/FlutterTentaclePlugin.swift
new file mode 100644
--- /dev/null
+++ b/Engine/Extras/FlutterTentacle/IOS/Classes/FlutterTentaclePlugin.swift
@@ -0,0 +1,12 @@
+// Copyright Epic Games, Inc. All Rights Reserved.
+
+import Flutter
+import UIKit
+
+public class FlutterTentaclePlugin: NSObject, FlutterPlugin {
+  static var tentacleApi: FlutterTentacleApi?
+  
+  public static func register(with registrar: FlutterPluginRegistrar) {
+    tentacleApi = FlutterTentacleApi(binaryMessenger: registrar.messenger())
+  }
+}
diff --git a/Engine/Extras/FlutterTentacle/IOS/Classes/Util/FlutterPluginApi.swift b/Engine/Extras/FlutterTentacle/IOS/Classes/Util/FlutterPluginApi.swift
new file mode 100644
--- /dev/null
+++ b/Engine/Extras/FlutterTentacle/IOS/Classes/Util/FlutterPluginApi.swift
@@ -0,0 +1,21 @@
+// Copyright Epic Games, Inc. All Rights Reserved.
+
+import Foundation
+
+/// Base class for all APIs that communicate to and from a Flutter plugin.
+class FlutterPluginApi<T>: NSObject {
+  /// The API used to send messages to Flutter.
+  private let flutter: T
+  
+  init (flutter: T) {
+    self.flutter = flutter
+  }
+  
+  /// Call a function on the Flutter API.
+  /// This will run on the main thread as required for any calls from platform to Flutter.
+  func callFlutter(call: @escaping (T) -> Void) {
+    DispatchQueue.main.async {
+      call(self.flutter)
+    }
+  }
+}
diff --git a/Engine/Extras/FlutterTentacle/IOS/Classes/Util/TentacleExtensions.swift b/Engine/Extras/FlutterTentacle/IOS/Classes/Util/TentacleExtensions.swift
new file mode 100644
--- /dev/null
+++ b/Engine/Extras/FlutterTentacle/IOS/Classes/Util/TentacleExtensions.swift
@@ -0,0 +1,14 @@
+// Copyright Epic Games, Inc. All Rights Reserved.
+
+import Tentacle
+
+extension Tentacle.TentacleProductId {
+  /// Convert from a native Tentacle product ID to the Flutter equivalent.
+  func toFlutter() -> TentacleProductId {
+    switch (self) {
+    case TentacleProductIdSyncE: return .syncE
+    case TentacleProductIdTrackE: return .trackE
+    default: return .generic
+    }
+  }
+}
diff --git a/Engine/Extras/FlutterTentacle/IOSSetup.sh b/Engine/Extras/FlutterTentacle/IOSSetup.sh
new file mode 100755
--- /dev/null
+++ b/Engine/Extras/FlutterTentacle/IOSSetup.sh
@@ -0,0 +1,6 @@
+#!/bin/bash
+
+SYMLINK_PATH=IOS/Tentacle.xcframework
+
+rm -f $SYMLINK_PATH
+ln -s ../../../Restricted/NotForLicensees/Source/ThirdParty/TentacleSDK/apple/build/Tentacle.xcframework $SYMLINK_PATH
diff --git a/Engine/Extras/FlutterTentacle/README.md b/Engine/Extras/FlutterTentacle/README.md
new file mode 100644
--- /dev/null
+++ b/Engine/Extras/FlutterTentacle/README.md
@@ -0,0 +1,116 @@
+# flutter_tentacle
+
+This package adds Flutter support for Bluetooth scanning and timecode synchronization with Tentacle devices. It is designed for use specifically in Epic's Flutter apps, so it depends on the `EpicCommon` package and integrates with its features and related tooling.
+
+## Setup
+
+### Tentacle SDK
+
+This project depends on the Tentacle SDK located in `Engine/Restricted/NotForLicensees/Source/ThirdParty/TentacleSDK`. Note that this SDK is not
+distributed to Unreal Engine licensees. If you have access to it, place it in that location to build this package.
+
+### Android Configuration
+
+In the main `AndroidManifest.xml' for your project, add the following elements to the `<manifest>` element:
+
+	<!-- Features and permissions for Tentacle scanning -->
+    <uses-feature
+        android:name="android.hardware.bluetooth_le"
+        android:required="false" />
+
+	<!-- API <= 30 -->
+	<uses-permission
+		android:name="android.permission.BLUETOOTH"
+		android:maxSdkVersion="30" />
+	<uses-permission
+		android:name="android.permission.BLUETOOTH_ADMIN"
+		android:maxSdkVersion="30" />
+	<uses-permission android:name="android.permission.ACCESS_COARSE_LOCATION" />
+	<uses-permission android:name="android.permission.ACCESS_FINE_LOCATION" />
+
+	<!-- API >= 31 -->
+	<uses-permission
+		android:name="android.permission.BLUETOOTH_SCAN"
+		android:usesPermissionFlags="neverForLocation"
+		tools:targetApi="s" />
+	<uses-permission android:name="android.permission.BLUETOOTH_CONNECT" />
+	<!-- End Tentacle permissions -->
+
+If your build fails when trying to parse this file, make sure this attribute is present on your `<manifest>` element:
+
+	xmlns:tools="http://schemas.android.com/tools"
+
+### iOS Configuration
+
+First, run the `IOSSetup.sh` script in this directory. This will create a symbolic link to the Tentacle framework in the IOS directory, which is required to build this library.
+
+Open the `Podfile` file in your Flutter app's `ios` directory and locate the code block starting with: `post_install do |installer|`. Above that line, add the following function.
+
+	# Configure the Pods project with build information necessary for flutter_tentacle
+	# Configure the Pods project with build information necessary for flutter_tentacle
+	def configure_flutter_tentacle(installer)
+	  tentacle_target = installer.pods_project.targets.find { |target| target.name == 'flutter_tentacle' }
+	  
+	  if tentacle_target
+		# Allow Tentacle library to use Bluetooth permissions for device communication
+		tentacle_target.build_configurations.each do |config|
+		  config.build_settings['GCC_PREPROCESSOR_DEFINITIONS'] ||= [
+			'$(inherited)',
+			'PERMISSION_BLUETOOTH=1',
+		  ]
+		end
+		
+		dev_pods_group = installer.pods_project.groups.find { |group| group.display_name == 'Development Pods' }
+		framework_ref = dev_pods_group['flutter_tentacle']['Frameworks']['Tentacle.xcframework']
+	
+		# Embed the Tentacle framework
+		embed_frameworks_build_phase = installer.pods_project.new(Xcodeproj::Project::Object::PBXCopyFilesBuildPhase)
+		embed_frameworks_build_phase.name = 'Embed Frameworks'
+		embed_frameworks_build_phase.symbol_dst_subfolder_spec = :frameworks
+		tentacle_target.build_phases << embed_frameworks_build_phase
+	
+		# Add the framework as a Sign & Embed build file
+		build_file = embed_frameworks_build_phase.add_file_reference(framework_ref)
+		build_file.settings = { 'ATTRIBUTES' => ['CodeSignOnCopy', 'RemoveHeadersOnCopy'] }
+		frameworks_build_phase = tentacle_target.build_phases.find { |build_phase| build_phase.to_s == 'FrameworksBuildPhase' }
+		frameworks_build_phase.add_file_reference(framework_ref)
+	  end
+	end
+
+Then, inside the `post_install` block, add the following line:
+
+	configure_flutter_tentacle(installer)
+
+See `example/IOS/Podfile` for an example.
+
+### Flutter Configuration
+
+1. Add this package to your package's `pubspec.yaml`.
+2. In your main build function...
+	1. Add a top-level `Provider` for `TentacleDeviceManager`, e.g.:
+    ```
+    Provider(
+      create: (_) => TentacleDeviceManager(),
+      dispose: (_, provider) => provider.dispose(),
+    )
+    ```
+	2. Add `TentacleLocalizations.localizationsDelegates` to your list of localization delegates.
+3. When your app exits (e.g. in the `didRequestAppExit` function), call `TentaclePlugin.shutdown()`.
+4. Optional: to add Tentacle timecode support to the `TimecodeManager`, call `manager.registerSource(TentacleTimecodeSource())` before you call `manager.initialize()`.
+
+## Package development
+
+### ffigen
+
+We use some Tentacle SDK functions through Dart's Foreign Function Interface (FFI). Dart bindings for these functions are
+automatically generated using `ffigen` and checked into source control.
+
+If you need to regenerate them, first set up ffigen (see https://pub.dev/packages/ffigen), then run the following command
+in the root VCAM project directory:
+
+    dart run ffigen --config ffigen_tentacle.yaml
+
+### Pigeon
+
+This project uses Pigeon to automatically generate platform channel bindings for native APIs.
+See `pigeons/README` for more information.
\ No newline at end of file
diff --git a/Engine/Extras/FlutterTentacle/analysis_options.yaml b/Engine/Extras/FlutterTentacle/analysis_options.yaml
new file mode 100644
--- /dev/null
+++ b/Engine/Extras/FlutterTentacle/analysis_options.yaml
@@ -0,0 +1,4 @@
+include: package:flutter_lints/flutter.yaml
+
+# Additional information about this file can be found at
+# https://dart.dev/guides/language/analysis-options
diff --git a/Engine/Extras/FlutterTentacle/android/.gitignore b/Engine/Extras/FlutterTentacle/android/.gitignore
new file mode 100644
--- /dev/null
+++ b/Engine/Extras/FlutterTentacle/android/.gitignore
@@ -0,0 +1,9 @@
+*.iml
+.gradle
+/local.properties
+/.idea/workspace.xml
+/.idea/libraries
+.DS_Store
+/build
+/captures
+.cxx
diff --git a/Engine/Extras/FlutterTentacle/android/src/main/AndroidManifest.xml b/Engine/Extras/FlutterTentacle/android/src/main/AndroidManifest.xml
new file mode 100644
--- /dev/null
+++ b/Engine/Extras/FlutterTentacle/android/src/main/AndroidManifest.xml
@@ -0,0 +1,3 @@
+<manifest xmlns:android="http://schemas.android.com/apk/res/android"
+  package="com.epicgames.flutter_tentacle">
+</manifest>
diff --git a/Engine/Extras/FlutterTentacle/example/.gitignore b/Engine/Extras/FlutterTentacle/example/.gitignore
new file mode 100644
--- /dev/null
+++ b/Engine/Extras/FlutterTentacle/example/.gitignore
@@ -0,0 +1,43 @@
+# Miscellaneous
+*.class
+*.log
+*.pyc
+*.swp
+.DS_Store
+.atom/
+.buildlog/
+.history
+.svn/
+migrate_working_dir/
+
+# IntelliJ related
+*.iml
+*.ipr
+*.iws
+.idea/
+
+# The .vscode folder contains launch configuration and tasks you configure in
+# VS Code which you may wish to be included in version control, so this line
+# is commented out by default.
+#.vscode/
+
+# Flutter/Dart/Pub related
+**/doc/api/
+**/ios/Flutter/.last_build_id
+.dart_tool/
+.flutter-plugins
+.flutter-plugins-dependencies
+.pub-cache/
+.pub/
+/build/
+
+# Symbolication related
+app.*.symbols
+
+# Obfuscation related
+app.*.map.json
+
+# Android Studio will place build artifacts here
+/android/app/debug
+/android/app/profile
+/android/app/release
diff --git a/Engine/Extras/FlutterTentacle/example/IOS/.gitignore b/Engine/Extras/FlutterTentacle/example/IOS/.gitignore
new file mode 100644
--- /dev/null
+++ b/Engine/Extras/FlutterTentacle/example/IOS/.gitignore
@@ -0,0 +1,34 @@
+**/dgph
+*.mode1v3
+*.mode2v3
+*.moved-aside
+*.pbxuser
+*.perspectivev3
+**/*sync/
+.sconsign.dblite
+.tags*
+**/.vagrant/
+**/DerivedData/
+Icon?
+**/Pods/
+**/.symlinks/
+profile
+xcuserdata
+**/.generated/
+Flutter/App.framework
+Flutter/Flutter.framework
+Flutter/Flutter.podspec
+Flutter/Generated.xcconfig
+Flutter/ephemeral/
+Flutter/app.flx
+Flutter/app.zip
+Flutter/flutter_assets/
+Flutter/flutter_export_environment.sh
+ServiceDefinitions.json
+Runner/GeneratedPluginRegistrant.*
+
+# Exceptions to above rules.
+!default.mode1v3
+!default.mode2v3
+!default.pbxuser
+!default.perspectivev3
diff --git a/Engine/Extras/FlutterTentacle/example/IOS/Flutter/Debug.xcconfig b/Engine/Extras/FlutterTentacle/example/IOS/Flutter/Debug.xcconfig
new file mode 100644
--- /dev/null
+++ b/Engine/Extras/FlutterTentacle/example/IOS/Flutter/Debug.xcconfig
@@ -0,0 +1,2 @@
+#include? "Pods/Target Support Files/Pods-Runner/Pods-Runner.debug.xcconfig"
+#include "Generated.xcconfig"
diff --git a/Engine/Extras/FlutterTentacle/example/IOS/Flutter/Release.xcconfig b/Engine/Extras/FlutterTentacle/example/IOS/Flutter/Release.xcconfig
new file mode 100644
--- /dev/null
+++ b/Engine/Extras/FlutterTentacle/example/IOS/Flutter/Release.xcconfig
@@ -0,0 +1,2 @@
+#include? "Pods/Target Support Files/Pods-Runner/Pods-Runner.release.xcconfig"
+#include "Generated.xcconfig"
diff --git a/Engine/Extras/FlutterTentacle/example/IOS/Runner/AppDelegate.swift b/Engine/Extras/FlutterTentacle/example/IOS/Runner/AppDelegate.swift
new file mode 100644
--- /dev/null
+++ b/Engine/Extras/FlutterTentacle/example/IOS/Runner/AppDelegate.swift
@@ -0,0 +1,13 @@
+import Flutter
+import UIKit
+
+@main
+@objc class AppDelegate: FlutterAppDelegate {
+  override func application(
+    _ application: UIApplication,
+    didFinishLaunchingWithOptions launchOptions: [UIApplication.LaunchOptionsKey: Any]?
+  ) -> Bool {
+    GeneratedPluginRegistrant.register(with: self)
+    return super.application(application, didFinishLaunchingWithOptions: launchOptions)
+  }
+}
diff --git a/Engine/Extras/FlutterTentacle/example/IOS/Runner/Assets.xcassets/AppIcon.appiconset/Contents.json b/Engine/Extras/FlutterTentacle/example/IOS/Runner/Assets.xcassets/AppIcon.appiconset/Contents.json
new file mode 100644
--- /dev/null
+++ b/Engine/Extras/FlutterTentacle/example/IOS/Runner/Assets.xcassets/AppIcon.appiconset/Contents.json
@@ -0,0 +1,122 @@
+{
+  "images" : [
+    {
+      "size" : "20x20",
+      "idiom" : "iphone",
+      "filename" : "[REDACTED_EMAIL]",
+      "scale" : "2x"
+    },
+    {
+      "size" : "20x20",
+      "idiom" : "iphone",
+      "filename" : "[REDACTED_EMAIL]",
+      "scale" : "3x"
+    },
+    {
+      "size" : "29x29",
+      "idiom" : "iphone",
+      "filename" : "[REDACTED_EMAIL]",
+      "scale" : "1x"
+    },
+    {
+      "size" : "29x29",
+      "idiom" : "iphone",
+      "filename" : "[REDACTED_EMAIL]",
+      "scale" : "2x"
+    },
+    {
+      "size" : "29x29",
+      "idiom" : "iphone",
+      "filename" : "[REDACTED_EMAIL]",
+      "scale" : "3x"
+    },
+    {
+      "size" : "40x40",
+      "idiom" : "iphone",
+      "filename" : "[REDACTED_EMAIL]",
+      "scale" : "2x"
+    },
+    {
+      "size" : "40x40",
+      "idiom" : "iphone",
+      "filename" : "[REDACTED_EMAIL]",
+      "scale" : "3x"
+    },
+    {
+      "size" : "60x60",
+      "idiom" : "iphone",
+      "filename" : "[REDACTED_EMAIL]",
+      "scale" : "2x"
+    },
+    {
+      "size" : "60x60",
+      "idiom" : "iphone",
+      "filename" : "[REDACTED_EMAIL]",
+      "scale" : "3x"
+    },
+    {
+      "size" : "20x20",
+      "idiom" : "ipad",
+      "filename" : "[REDACTED_EMAIL]",
+      "scale" : "1x"
+    },
+    {
+      "size" : "20x20",
+      "idiom" : "ipad",
+      "filename" : "[REDACTED_EMAIL]",
+      "scale" : "2x"
+    },
+    {
+      "size" : "29x29",
+      "idiom" : "ipad",
+      "filename" : "[REDACTED_EMAIL]",
+      "scale" : "1x"
+    },
+    {
+      "size" : "29x29",
+      "idiom" : "ipad",
+      "filename" : "[REDACTED_EMAIL]",
+      "scale" : "2x"
+    },
+    {
+      "size" : "40x40",
+      "idiom" : "ipad",
+      "filename" : "[REDACTED_EMAIL]",
+      "scale" : "1x"
+    },
+    {
+      "size" : "40x40",
+      "idiom" : "ipad",
+      "filename" : "[REDACTED_EMAIL]",
+      "scale" : "2x"
+    },
+    {
+      "size" : "76x76",
+      "idiom" : "ipad",
+      "filename" : "[REDACTED_EMAIL]",
+      "scale" : "1x"
+    },
+    {
+      "size" : "76x76",
+      "idiom" : "ipad",
+      "filename" : "[REDACTED_EMAIL]",
+      "scale" : "2x"
+    },
+    {
+      "size" : "83.5x83.5",
+      "idiom" : "ipad",
+      "filename" : "[REDACTED_EMAIL]",
+      "scale" : "2x"
+    },
+    {
+      "size" : "1024x1024",
+      "idiom" : "ios-marketing",
+      "filename" : "[REDACTED_EMAIL]",
+      "scale" : "1x"
+    }
+  ],
+  "info" : {
+    "version" : 1,
+    "author" : "xcode"
+  }
+}
diff --git a/Engine/Extras/FlutterTentacle/example/IOS/Runner/Assets.xcassets/LaunchImage.imageset/Contents.json b/Engine/Extras/FlutterTentacle/example/IOS/Runner/Assets.xcassets/LaunchImage.imageset/Contents.json
new file mode 100644
--- /dev/null
+++ b/Engine/Extras/FlutterTentacle/example/IOS/Runner/Assets.xcassets/LaunchImage.imageset/Contents.json
@@ -0,0 +1,23 @@
+{
+  "images" : [
+    {
+      "idiom" : "universal",
+      "filename" : "LaunchImage.png",
+      "scale" : "1x"
+    },
+    {
+      "idiom" : "universal",
+      "filename" : "[REDACTED_EMAIL]",
+      "scale" : "2x"
+    },
+    {
+      "idiom" : "universal",
+      "filename" : "[REDACTED_EMAIL]",
+      "scale" : "3x"
+    }
+  ],
+  "info" : {
+    "version" : 1,
+    "author" : "xcode"
+  }
+}
diff --git a/Engine/Extras/FlutterTentacle/example/IOS/Runner/Assets.xcassets/LaunchImage.imageset/README.md b/Engine/Extras/FlutterTentacle/example/IOS/Runner/Assets.xcassets/LaunchImage.imageset/README.md
new file mode 100644
--- /dev/null
+++ b/Engine/Extras/FlutterTentacle/example/IOS/Runner/Assets.xcassets/LaunchImage.imageset/README.md
@@ -0,0 +1,5 @@
+# Launch Screen Assets
+
+You can customize the launch screen with your own desired assets by replacing the image files in this directory.
+
+You can also do it by opening your Flutter project's Xcode project with `open ios/Runner.xcworkspace`, selecting `Runner/Assets.xcassets` in the Project Navigator and dropping in the desired images.
\ No newline at end of file
diff --git a/Engine/Extras/FlutterTentacle/example/IOS/Runner/Runner-Bridging-Header.h b/Engine/Extras/FlutterTentacle/example/IOS/Runner/Runner-Bridging-Header.h
new file mode 100644
--- /dev/null
+++ b/Engine/Extras/FlutterTentacle/example/IOS/Runner/Runner-Bridging-Header.h
@@ -0,0 +1 @@
+#import "GeneratedPluginRegistrant.h"
diff --git a/Engine/Extras/FlutterTentacle/example/IOS/RunnerTests/RunnerTests.swift b/Engine/Extras/FlutterTentacle/example/IOS/RunnerTests/RunnerTests.swift
new file mode 100644
--- /dev/null
+++ b/Engine/Extras/FlutterTentacle/example/IOS/RunnerTests/RunnerTests.swift
@@ -0,0 +1,27 @@
+import Flutter
+import UIKit
+import XCTest
+
+
+@testable import flutter_tentacle
+
+// This demonstrates a simple unit test of the Swift portion of this plugin's implementation.
+//
+// See https://developer.apple.com/documentation/xctest for more information about using XCTest.
+
+class RunnerTests: XCTestCase {
+
+  func testGetPlatformVersion() {
+    let plugin = FlutterTentaclePlugin()
+
+    let call = FlutterMethodCall(methodName: "getPlatformVersion", arguments: [])
+
+    let resultExpectation = expectation(description: "result block must be called.")
+    plugin.handle(call) { result in
+      XCTAssertEqual(result as! String, "iOS " + UIDevice.current.systemVersion)
+      resultExpectation.fulfill()
+    }
+    waitForExpectations(timeout: 1)
+  }
+
+}
diff --git a/Engine/Extras/FlutterTentacle/example/README.md b/Engine/Extras/FlutterTentacle/example/README.md
new file mode 100644
--- /dev/null
+++ b/Engine/Extras/FlutterTentacle/example/README.md
@@ -0,0 +1,16 @@
+# flutter_tentacle_example
+
+Demonstrates how to use the flutter_tentacle plugin.
+
+## Getting Started
+
+This project is a starting point for a Flutter application.
+
+A few resources to get you started if this is your first Flutter project:
+
+- [Lab: Write your first Flutter app](https://docs.flutter.dev/get-started/codelab)
+- [Cookbook: Useful Flutter samples](https://docs.flutter.dev/cookbook)
+
+For help getting started with Flutter development, view the
+[online documentation](https://docs.flutter.dev/), which offers tutorials,
+samples, guidance on mobile development, and a full API reference.
diff --git a/Engine/Extras/FlutterTentacle/example/analysis_options.yaml b/Engine/Extras/FlutterTentacle/example/analysis_options.yaml
new file mode 100644
--- /dev/null
+++ b/Engine/Extras/FlutterTentacle/example/analysis_options.yaml
@@ -0,0 +1,28 @@
+# This file configures the analyzer, which statically analyzes Dart code to
+# check for errors, warnings, and lints.
+#
+# The issues identified by the analyzer are surfaced in the UI of Dart-enabled
+# IDEs (https://dart.dev/tools#ides-and-editors). The analyzer can also be
+# invoked from the command line by running `flutter analyze`.
+
+# The following line activates a set of recommended lints for Flutter apps,
+# packages, and plugins designed to encourage good coding practices.
+include: package:flutter_lints/flutter.yaml
+
+linter:
+  # The lint rules applied to this project can be customized in the
+  # section below to disable rules from the `package:flutter_lints/flutter.yaml`
+  # included above or to enable additional rules. A list of all available lints
+  # and their documentation is published at https://dart.dev/lints.
+  #
+  # Instead of disabling a lint rule for the entire project in the
+  # section below, it can also be suppressed for a single line of code
+  # or a specific dart file by using the `// ignore: name_of_lint` and
+  # `// ignore_for_file: name_of_lint` syntax on the line or in the file
+  # producing the lint.
+  rules:
+    # avoid_print: false  # Uncomment to disable the `avoid_print` rule
+    # prefer_single_quotes: true  # Uncomment to enable the `prefer_single_quotes` rule
+
+# Additional information about this file can be found at
+# https://dart.dev/guides/language/analysis-options
diff --git a/Engine/Extras/FlutterTentacle/example/android/.gitignore b/Engine/Extras/FlutterTentacle/example/android/.gitignore
new file mode 100644
--- /dev/null
+++ b/Engine/Extras/FlutterTentacle/example/android/.gitignore
@@ -0,0 +1,13 @@
+gradle-wrapper.jar
+/.gradle
+/captures/
+/gradlew
+/gradlew.bat
+/local.properties
+GeneratedPluginRegistrant.java
+
+# Remember to never publicly share your keystore.
+# See https://flutter.dev/to/reference-keystore
+key.properties
+**/*.keystore
+**/*.jks
diff --git a/Engine/Extras/FlutterTentacle/example/pubspec.yaml b/Engine/Extras/FlutterTentacle/example/pubspec.yaml
new file mode 100644
--- /dev/null
+++ b/Engine/Extras/FlutterTentacle/example/pubspec.yaml
@@ -0,0 +1,88 @@
+name: flutter_tentacle_example
+description: "Demonstrates how to use the flutter_tentacle plugin."
+# The following line prevents the package from being accidentally published to
+# pub.dev using `flutter pub publish`. This is preferred for private packages.
+publish_to: 'none' # Remove this line if you wish to publish to pub.dev
+
+environment:
+  sdk: ^3.5.4
+
+# Dependencies specify other packages that your package needs in order to work.
+# To automatically upgrade your package dependencies to the latest versions
+# consider running `flutter pub upgrade --major-versions`. Alternatively,
+# dependencies can be manually updated by changing the version numbers below to
+# the latest version available on pub.dev. To see which dependencies have newer
+# versions available, run `flutter pub outdated`.
+dependencies:
+  flutter:
+    sdk: flutter
+
+  epic_common:
+    path: ../../FlutterCommon
+
+  flutter_tentacle:
+    # When depending on this package from a real application you should use:
+    #   flutter_tentacle: ^x.y.z
+    # See https://dart.dev/tools/pub/dependencies#version-constraints
+    # The example app is bundled with the plugin so we use a path dependency on
+    # the parent directory to use the current plugin's version.
+    path: ../
+
+  # The following adds the Cupertino Icons font to your application.
+  # Use with the CupertinoIcons class for iOS style icons.
+  cupertino_icons: ^1.0.8
+
+dev_dependencies:
+  integration_test:
+    sdk: flutter
+  flutter_test:
+    sdk: flutter
+
+  # The "flutter_lints" package below contains a set of recommended lints to
+  # encourage good coding practices. The lint set provided by the package is
+  # activated in the `analysis_options.yaml` file located at the root of your
+  # package. See that file for information about deactivating specific lint
+  # rules and activating additional ones.
+  flutter_lints: ^4.0.0
+
+# For information on the generic Dart part of this file, see the
+# following page: https://dart.dev/tools/pub/pubspec
+
+# The following section is specific to Flutter packages.
+flutter:
+
+  # The following line ensures that the Material Icons font is
+  # included with your application, so that you can use the icons in
+  # the material Icons class.
+  uses-material-design: true
+
+  # To add assets to your application, add an assets section, like this:
+  # assets:
+  #   - images/a_dot_burr.jpeg
+  #   - images/a_dot_ham.jpeg
+
+  # An image asset can refer to one or more resolution-specific "variants", see
+  # https://flutter.dev/to/resolution-aware-images
+
+  # For details regarding adding assets from package dependencies, see
+  # https://flutter.dev/to/asset-from-package
+
+  # To add custom fonts to your application, add a fonts section here,
+  # in this "flutter" section. Each entry in this list should have a
+  # "family" key with the font family name, and a "fonts" key with a
+  # list giving the asset and other descriptors for the font. For
+  # example:
+  # fonts:
+  #   - family: Schyler
+  #     fonts:
+  #       - asset: fonts/Schyler-Regular.ttf
+  #       - asset: fonts/Schyler-Italic.ttf
+  #         style: italic
+  #   - family: Trajan Pro
+  #     fonts:
+  #       - asset: fonts/TrajanPro.ttf
+  #       - asset: fonts/TrajanPro_Bold.ttf
+  #         weight: 700
+  #
+  # For details regarding fonts from package dependencies,
+  # see https://flutter.dev/to/font-from-package
diff --git a/Engine/Extras/FlutterTentacle/l10n.yaml b/Engine/Extras/FlutterTentacle/l10n.yaml
new file mode 100644
--- /dev/null
+++ b/Engine/Extras/FlutterTentacle/l10n.yaml
@@ -0,0 +1,7 @@
+arb-dir: lib/src/l10n
+template-arb-file: tentacle_en.arb
+synthetic-package: false
+header-file: generated_header.txt
+output-class: TentacleLocalizations
+output-dir: lib/src/l10n/generated
+output-localization-file: tentacle_localizations.dart
\ No newline at end of file
diff --git a/Engine/Extras/FlutterTentacle/lib/src/l10n/generated_header.txt b/Engine/Extras/FlutterTentacle/lib/src/l10n/generated_header.txt
new file mode 100644
--- /dev/null
+++ b/Engine/Extras/FlutterTentacle/lib/src/l10n/generated_header.txt
@@ -0,0 +1,2 @@
+// Copyright Epic Games, Inc. All Rights Reserved.
+// Generated automatically by flutter_gen.
diff --git a/Engine/Extras/FlutterTentacle/pigeons/README b/Engine/Extras/FlutterTentacle/pigeons/README
new file mode 100644
--- /dev/null
+++ b/Engine/Extras/FlutterTentacle/pigeons/README
@@ -0,0 +1,15 @@
+Dart files in this directory are not directly included in the app's source -- instead, they will be
+processed by Pigeon to generate code that facilitates communication between Flutter and each of
+the platform languages.
+
+You can modify the files in this directory to change the cross-platform API. Each file has an
+"@ConfigurePigeon" annotation which specifies where the resulting generated files will be stored
+for each target language. These generated files *should* be checked into Perforce, as they are not
+automatically generated each time the app is run.
+
+To re-generate Pigeon output, run this in your terminal from the app's root directory:
+
+  dart run pigeon --input pigeons/<filename>.dart
+
+For more information about Pigeon, see:
+https://pub.dev/packages/pigeon
\ No newline at end of file
diff --git a/Engine/Extras/FlutterTentacle/pubspec.yaml b/Engine/Extras/FlutterTentacle/pubspec.yaml
new file mode 100644
--- /dev/null
+++ b/Engine/Extras/FlutterTentacle/pubspec.yaml
@@ -0,0 +1,53 @@
+name: flutter_tentacle
+description: "Adds support for Tentacle devices"
+publish_to: 'none'
+version: 0.0.1
+
+environment:
+  sdk: '>=3.1.0 <4.0.0'
+  flutter: "^3.16.3"
+
+dependencies:
+  flutter:
+    sdk: flutter
+
+  flutter_localizations:
+    sdk: flutter
+    
+  epic_common:
+    path: ../FlutterCommon
+    
+  streaming_shared_preferences:
+    path: ../ThirdPartyNotUE/streaming_shared_preferences
+      
+  device_info_plus: ^11.1.0
+  ffi: ^2.1.0
+  intl: any
+  logging: ^1.2.0
+  permission_handler: ^11.1.0
+  plugin_platform_interface: ^2.0.2
+  provider: ^6.0.2
+  system_clock: ^2.0.0
+
+dev_dependencies:
+  flutter_test:
+    sdk: flutter
+  flutter_lints: ^4.0.0
+
+  pigeon: ^22.6.3
+  ffigen: ^16.0.0
+
+flutter:
+  generate: true
+  
+  assets:
+    - assets/icons/
+    - assets/icons/tentacle/
+  
+  plugin:
+    platforms:
+      android:
+        package: com.epicgames.flutter_tentacle
+        pluginClass: FlutterTentaclePlugin
+      ios:
+        pluginClass: FlutterTentaclePlugin
\ No newline at end of file
diff --git a/Engine/Extras/VirtualProduction/FlutterLiveLinkVCAM/IOS/Runner/AppDelegate.swift b/Engine/Extras/VirtualProduction/FlutterLiveLinkVCAM/IOS/Runner/AppDelegate.swift
--- a/Engine/Extras/VirtualProduction/FlutterLiveLinkVCAM/IOS/Runner/AppDelegate.swift
+++ b/Engine/Extras/VirtualProduction/FlutterLiveLinkVCAM/IOS/Runner/AppDelegate.swift
@@ -11,7 +11,6 @@ import SpriteKit
   var dataChannelApi: FlutterRtcDataChannelApi?
   var videoViewControllerApi: FlutterRtcVideoViewControllerApi?
   var arSessionApi: FlutterArSessionApi?
-  var tentacleApi: FlutterTentacleApi?
   var gamepadApi: FlutterGamepadApi?
   
   override func application(
@@ -30,7 +29,6 @@ import SpriteKit
     dataChannelApi = FlutterRtcDataChannelApi(binaryMessenger: binaryMessenger)
     videoViewControllerApi = FlutterRtcVideoViewControllerApi(binaryMessenger: binaryMessenger, textureRegistry: flutterViewController)
     arSessionApi = FlutterArSessionApi(binaryMessenger: binaryMessenger)
-    tentacleApi = FlutterTentacleApi(binaryMessenger: binaryMessenger)
     gamepadApi = FlutterGamepadApi(binaryMessenger: binaryMessenger)
   
     return super.application(application, didFinishLaunchingWithOptions: launchOptions)
diff --git a/Engine/Extras/VirtualProduction/FlutterLiveLinkVCAM/IOS/Runner/Tentacle/API/FlutterTentacleApi.swift b/Engine/Extras/VirtualProduction/FlutterLiveLinkVCAM/IOS/Runner/Tentacle/API/FlutterTentacleApi.swift
deleted file mode 100644
diff --git a/Engine/Extras/VirtualProduction/FlutterLiveLinkVCAM/IOS/Runner/Tentacle/API/TentacleApi.g.swift b/Engine/Extras/VirtualProduction/FlutterLiveLinkVCAM/IOS/Runner/Tentacle/API/TentacleApi.g.swift
deleted file mode 100644
diff --git a/Engine/Extras/VirtualProduction/FlutterLiveLinkVCAM/IOS/Runner/Tentacle/Util/TentacleExtensions.swift b/Engine/Extras/VirtualProduction/FlutterLiveLinkVCAM/IOS/Runner/Tentacle/Util/TentacleExtensions.swift
deleted file mode 100644
diff --git a/Engine/Extras/VirtualProduction/FlutterLiveLinkVCAM/README.md b/Engine/Extras/VirtualProduction/FlutterLiveLinkVCAM/README.md
--- a/Engine/Extras/VirtualProduction/FlutterLiveLinkVCAM/README.md
+++ b/Engine/Extras/VirtualProduction/FlutterLiveLinkVCAM/README.md
@@ -2,28 +2,24 @@
 
 A rebuild of the Live Link VCAM app for Flutter, enabling cross-platform deployment.
 
-## Tentacle Integration
+## Tentacle
 
-This project depends on the Tentacle SDK located in Engine/Restricted/NotForLicensees/Source/ThirdParty/TentacleSDK.
+This app depends on the `FlutterTentacle` package (included in the Unreal Engine repo under `Engine/Extras/FlutterTentacle`), which in turn depends on the Tentacle SDK. The Tentacle SDK is not distributed to engine licensees. See the `FlutterTentacle` README for more info.
 
-We use some Tentacle SDK functions through Dart's Foreign Function Interface (FFI). Dart bindings for these functions are
-automatically generated using ffigen and checked into source control.
+### IOS Note
 
-If you need to regenerate them, first set up ffigen (see https://pub.dev/packages/ffigen), then run the following command
-in the root VCAM project directory:
-
-    `dart run ffigen --config ffigen_tentacle.yaml`
+Before building `live_link_vcam` for iOS, you will first need to run the `IOSSetup.sh` script as described in the `FlutterTentacle` README.
 
 ## Pigeon
 
 This project uses Pigeon to automatically generate platform channel bindings for native APIs.
 See `pigeons/README` for more information.
 
-## Signing Certificates 
+## Signing Certificates
 
-To setup signing certificates, follow the instructions outlined in this article 
+To setup signing certificates, follow the instructions outlined in this article
 
- - [Deploying Flutter Apps to The PlayStore](https:[REDACTED_EMAIL]/deploying-flutter-apps-to-the-playstore-1bd0cce0d15c)
+- [Deploying Flutter Apps to The PlayStore](https:[REDACTED_EMAIL]/deploying-flutter-apps-to-the-playstore-1bd0cce0d15c)
 
 Add your `android-key.tks` to `android\app` folder and `key.properities` to `android\`.  The gradle files are already setup
 to read these files when app bundling
\ No newline at end of file
diff --git a/Engine/Extras/VirtualProduction/FlutterLiveLinkVCAM/pubspec.yaml b/Engine/Extras/VirtualProduction/FlutterLiveLinkVCAM/pubspec.yaml
--- a/Engine/Extras/VirtualProduction/FlutterLiveLinkVCAM/pubspec.yaml
+++ b/Engine/Extras/VirtualProduction/FlutterLiveLinkVCAM/pubspec.yaml
@@ -18,6 +18,9 @@ dependencies:
   epic_common:
     path: ../../FlutterCommon
 
+  flutter_tentacle:
+    path: ../../FlutterTentacle
+
   streaming_shared_preferences:
     path: ../../ThirdPartyNotUE/streaming_shared_preferences
 
@@ -32,23 +35,16 @@ dependencies:
   system_clock: ^2.0.0
   web_socket_channel: ^3.0.1
   share_plus: ^10.1.2
+  url_launcher: ^6.2.2
 
 dev_dependencies:
   flutter_test:
     sdk: flutter
+  flutter_lints: ^4.0.0
 
-  ffigen: ^16.0.0
   pigeon: ^22.6.3
-  url_launcher: ^6.2.2
   flutter_launcher_icons: ^0.14.1
 
-  # The "flutter_lints" package below contains a set of recommended lints to
-  # encourage good coding practices. The lint set provided by the package is
-  # activated in the `analysis_options.yaml` file located at the root of your
-  # package. See that file for information about deactivating specific lint
-  # rules and activating additional ones.
-  flutter_lints: ^5.0.0
-
 # For information on the generic Dart part of this file, see the
 # following page: https://dart.dev/tools/pub/pubspec
 
@@ -62,7 +58,6 @@ flutter:
   assets:
     - assets/eula.html
     - assets/icons/
-    - assets/icons/tentacle/
 
   fonts:
     - family: Inter
